Add tests for websockets in WebTests;
Move WebTests web server from Engine/Source/Programs/WebTests/WebServer to Engine/Source/Programs/WebTestsServer, so to make sure not include all files underneath into VS project.

diff --git a/Engine/Source/Programs/AutomationTool/LowLevelTests/Tests/LowLevelTests.WebTestsLowLevelTestsExtension.cs b/Engine/Source/Programs/AutomationTool/LowLevelTests/Tests/LowLevelTests.WebTestsLowLevelTestsExtension.cs
--- a/Engine/Source/Programs/AutomationTool/LowLevelTests/Tests/LowLevelTests.WebTestsLowLevelTestsExtension.cs
+++ b/Engine/Source/Programs/AutomationTool/LowLevelTests/Tests/LowLevelTests.WebTestsLowLevelTestsExtension.cs
@@ -30,7 +30,7 @@ namespace LowLevelTests
 		public void PreRunTests()
 		{
 			ProcessStartInfo StartInfo = new ProcessStartInfo();
-			StartInfo.WorkingDirectory = Path.Combine(Unreal.EngineDirectory.FullName, "Source", "Programs", "WebTests", "WebServer");
+			StartInfo.WorkingDirectory = Path.Combine(Unreal.EngineDirectory.FullName, "Source", "Programs", "WebTestsServer");
 			StartInfo.FileName = RuntimePlatform.IsWindows ? "cmd.exe" : "/bind/sh";
 			StartInfo.Arguments = RuntimePlatform.IsWindows ? "/c runserver.bat" : "-c 'runserver.sh'";
 			StartInfo.UseShellExecute = false;
diff --git a/Engine/Source/Programs/WebTests/Private/TestHttp.cpp b/Engine/Source/Programs/WebTests/Private/TestHttp.cpp
--- a/Engine/Source/Programs/WebTests/Private/TestHttp.cpp
+++ b/Engine/Source/Programs/WebTests/Private/TestHttp.cpp
@@ -32,7 +32,7 @@ class FHttpModuleTestFixture
 public:
 	FHttpModuleTestFixture()
 		: WebServerIp(TEXT("127.0.0.1"))
-		, WebServerPort(8000)
+		, WebServerHttpPort(8000)
 		, bRunHeavyTests(false)
 	{
 		ParseSettingsFromCommandLine();
@@ -56,13 +56,13 @@ public:
 	}
 
 	const FString UrlWithInvalidPortToTestConnectTimeout() const { return FString::Format(TEXT("http://{0}:{1}"), { *WebServerIp, 8765 }); }
-	const FString UrlBase() const { return FString::Format(TEXT("http://{0}:{1}"), { *WebServerIp, WebServerPort }); }
+	const FString UrlBase() const { return FString::Format(TEXT("http://{0}:{1}"), { *WebServerIp, WebServerHttpPort }); }
 	const FString UrlHttpTests() const { return FString::Format(TEXT("{0}/webtests/httptests"), { *UrlBase() }); }
 	const FString UrlToTestMethods() const { return FString::Format(TEXT("{0}/methods"), { *UrlHttpTests() }); }
 	const FString UrlStreamDownload(uint32 Chunks, uint32 ChunkSize) { return FString::Format(TEXT("{0}/streaming_download/{1}/{2}/"), { *UrlHttpTests(), Chunks, ChunkSize }); }
 
 	FString WebServerIp;
-	uint32 WebServerPort;
+	uint32 WebServerHttpPort;
 	FHttpModule* HttpModule;
 
 	bool bRunHeavyTests;
diff --git a/Engine/Source/Programs/WebTests/Private/TestWebSockets.cpp b/Engine/Source/Programs/WebTests/Private/TestWebSockets.cpp
new file mode 100644
--- /dev/null
+++ b/Engine/Source/Programs/WebTests/Private/TestWebSockets.cpp
@@ -0,0 +1,112 @@
+// Copyright Epic Games, Inc. All Rights Reserved.
+#include "Containers/BackgroundableTicker.h"
+#include "CoreMinimal.h"
+#include "HAL/PlatformProcess.h"
+#include "Misc/CommandLine.h"
+#include "WebSocketsModule.h"
+#include "IWebSocket.h"
+#include "TestHarness.h"
+
+/**
+ *  WebSockets Tests
+ *  -----------------------------------------------------------------------------------------------
+ *
+ *  PURPOSE:
+ *
+ *	Integration Tests to make sure all kinds of WebSockets client features in C++ work well on different platforms,
+ *  including but not limited to error handing, retrying, threading, SSL and profiling.
+ * 
+ *  -----------------------------------------------------------------------------------------------
+ */
+
+#define WEBSOCKETS_TAG "[WebSockets]"
+
+class FWebSocketsModuleTestFixture
+{
+public:
+	FWebSocketsModuleTestFixture()
+		: WebServerIp(TEXT("127.0.0.1"))
+		, WebServerWebSocketsPort(8000)
+	{
+		ParseSettingsFromCommandLine();
+
+		WebSocketsModule = new FWebSocketsModule();
+		IModuleInterface* Module = WebSocketsModule;
+		Module->StartupModule();
+	}
+
+	virtual ~FWebSocketsModuleTestFixture()
+	{
+		IModuleInterface* Module = WebSocketsModule;
+		Module->ShutdownModule();
+		delete Module;
+	}
+
+	void ParseSettingsFromCommandLine()
+	{
+		FParse::Value(FCommandLine::Get(), TEXT("web_server_ip"), WebServerIp);
+	}
+
+	const FString UrlWithInvalidPortToTestConnectTimeout() const { return FString::Format(TEXT("ws://{0}:{1}"), { *WebServerIp, 8765 }); }
+	const FString UrlBase() const { return FString::Format(TEXT("ws://{0}:{1}"), { *WebServerIp, WebServerWebSocketsPort }); }
+	const FString UrlWebSocketsTests() const { return FString::Format(TEXT("{0}/webtests/websocketstests"), { *UrlBase() }); }
+
+	FString WebServerIp;
+	uint32 WebServerWebSocketsPort;
+	FWebSocketsModule* WebSocketsModule;
+};
+
+class FRunUntilQuitRequestedFixture : public FWebSocketsModuleTestFixture
+{
+public:
+	FRunUntilQuitRequestedFixture()
+	{
+	}
+
+	~FRunUntilQuitRequestedFixture()
+	{
+		RunUntilQuitRequested();
+	}
+
+	void RunUntilQuitRequested()
+	{
+		while (!bQuitRequested)
+		{
+			FTSBackgroundableTicker::GetCoreTicker().Tick(TickFrequency);
+
+			FPlatformProcess::Sleep(TickFrequency);
+		}
+	}
+
+	float TickFrequency = 1.0f / 60; /*60 FPS*/;
+	bool bQuitRequested = false;
+	bool bSucceed = false;
+};
+
+TEST_CASE_METHOD(FRunUntilQuitRequestedFixture, "WebSockets can connect then send and receive message", WEBSOCKETS_TAG)
+{
+	TSharedPtr<IWebSocket> WebSocket = WebSocketsModule->CreateWebSocket(FString::Format(TEXT("{0}/echo/"), { *UrlWebSocketsTests() }));
+
+	WebSocket->OnConnected().AddLambda([this, WebSocket](){
+		WebSocket->Send(TEXT("hi websockets tests"));
+	});
+
+	WebSocket->OnMessage().AddLambda([this, WebSocket](const FString& MessageString) {
+		CHECK(MessageString == TEXT("hi websockets tests"));
+		WebSocket->Close();
+		bSucceed = true;
+	});
+
+	WebSocket->OnClosed().AddLambda([this, WebSocket](int32 /* StatusCode */, const FString& /* Reason */, bool /* bWasClean */){
+		CHECK(bSucceed);
+		bQuitRequested = true;
+	});
+
+	WebSocket->OnConnectionError().AddLambda([this, WebSocket](const FString& /* Error */){
+		CHECK(false);
+		bQuitRequested = true;
+	});
+
+	WebSocket->Connect();
+}
+
diff --git a/Engine/Source/Programs/WebTests/README.md b/Engine/Source/Programs/WebTests/README.md
--- a/Engine/Source/Programs/WebTests/README.md
+++ b/Engine/Source/Programs/WebTests/README.md
@@ -5,8 +5,8 @@ This is a collection of tests for web features like http, websocket, ssl etc.
 
 # Running the Tests
 
-## Run the WebServer(NOTE: This WebServer is created for TEST PURPOSE ONLY, it's not production ready!!!):
-	In ./WebServer folder:
+## Run the WebTestsServer(NOTE: This server is created for TEST PURPOSE ONLY, it's not production ready!!!):
+	In ../WebTestsServer folder:
 		> runserver.bat
 	Or if you use docker:
 		> dockerbuildandrun.bat
@@ -16,7 +16,7 @@ This is a collection of tests for web features like http, websocket, ssl etc.
 	If running tests on other devices, pass in the ip as command line args, after extra args AT THE END, like: "--extra-args --web_server_ip=your.pc.ip.address"
 	Compile and debug
 
-# Adding new test case in WebServer:
-	Add/change the code in ./WebServer code, most likely in httptests/urls.py and httptests/views.py, and save. Code will be reloaded if the web server is running
+# Adding new test case in WebTestsServer:
+	Add/change the code in ./WebTestsServer code, most likely in httptests/urls.py and httptests/views.py, and save. Code will be reloaded if the web server is running
 	When using docker, also need to run dockerbuildandrun.bat
 	For more info about how to code in django, check https://docs.djangoproject.com/en/
diff --git a/Engine/Source/Programs/WebTests/WebServer/dockerbuildandrun.bat b/Engine/Source/Programs/WebTests/WebServer/dockerbuildandrun.bat
deleted file mode 100755
diff --git a/Engine/Source/Programs/WebTests/WebServer/dockerfile b/Engine/Source/Programs/WebTests/WebServer/dockerfile
deleted file mode 100644
diff --git a/Engine/Source/Programs/WebTests/WebServer/httptests/__init__.py b/Engine/Source/Programs/WebTests/WebServer/httptests/__init__.py
deleted file mode 100644
diff --git a/Engine/Source/Programs/WebTests/WebServer/httptests/admin.py b/Engine/Source/Programs/WebTests/WebServer/httptests/admin.py
deleted file mode 100644
diff --git a/Engine/Source/Programs/WebTests/WebServer/httptests/apps.py b/Engine/Source/Programs/WebTests/WebServer/httptests/apps.py
deleted file mode 100644
diff --git a/Engine/Source/Programs/WebTests/WebServer/httptests/migrations/__init__.py b/Engine/Source/Programs/WebTests/WebServer/httptests/migrations/__init__.py
deleted file mode 100644
diff --git a/Engine/Source/Programs/WebTests/WebServer/httptests/models.py b/Engine/Source/Programs/WebTests/WebServer/httptests/models.py
deleted file mode 100644
diff --git a/Engine/Source/Programs/WebTests/WebServer/httptests/tests.py b/Engine/Source/Programs/WebTests/WebServer/httptests/tests.py
deleted file mode 100644
diff --git a/Engine/Source/Programs/WebTests/WebServer/httptests/urls.py b/Engine/Source/Programs/WebTests/WebServer/httptests/urls.py
deleted file mode 100644
diff --git a/Engine/Source/Programs/WebTests/WebServer/httptests/views.py b/Engine/Source/Programs/WebTests/WebServer/httptests/views.py
deleted file mode 100644
diff --git a/Engine/Source/Programs/WebTests/WebServer/manage.py b/Engine/Source/Programs/WebTests/WebServer/manage.py
deleted file mode 100644
diff --git a/Engine/Source/Programs/WebTests/WebServer/requirements.txt b/Engine/Source/Programs/WebTests/WebServer/requirements.txt
deleted file mode 100644
diff --git a/Engine/Source/Programs/WebTests/WebServer/runserver.bat b/Engine/Source/Programs/WebTests/WebServer/runserver.bat
deleted file mode 100755
diff --git a/Engine/Source/Programs/WebTests/WebServer/runserver.sh b/Engine/Source/Programs/WebTests/WebServer/runserver.sh
deleted file mode 100644
diff --git a/Engine/Source/Programs/WebTests/WebServer/webtests/__init__.py b/Engine/Source/Programs/WebTests/WebServer/webtests/__init__.py
deleted file mode 100644
diff --git a/Engine/Source/Programs/WebTests/WebServer/webtests/asgi.py b/Engine/Source/Programs/WebTests/WebServer/webtests/asgi.py
deleted file mode 100644
diff --git a/Engine/Source/Programs/WebTests/WebServer/webtests/settings.py b/Engine/Source/Programs/WebTests/WebServer/webtests/settings.py
deleted file mode 100644
diff --git a/Engine/Source/Programs/WebTests/WebServer/webtests/urls.py b/Engine/Source/Programs/WebTests/WebServer/webtests/urls.py
deleted file mode 100644
diff --git a/Engine/Source/Programs/WebTests/WebServer/webtests/wsgi.py b/Engine/Source/Programs/WebTests/WebServer/webtests/wsgi.py
deleted file mode 100644
diff --git a/Engine/Source/Programs/WebTests/WebTests.Build.cs b/Engine/Source/Programs/WebTests/WebTests.Build.cs
--- a/Engine/Source/Programs/WebTests/WebTests.Build.cs
+++ b/Engine/Source/Programs/WebTests/WebTests.Build.cs
@@ -11,6 +11,7 @@ public class WebTests : TestModuleRules
 				"ApplicationCore",
 				"Core",
 				"HTTP",
+				"WebSockets",
 				"Chaos",
 				"ChaosUserDataPT"
 			});
diff --git a/Engine/Source/Programs/WebTestsServer/dockerbuildandrun.bat b/Engine/Source/Programs/WebTestsServer/dockerbuildandrun.bat
new file mode 100755
--- /dev/null
+++ b/Engine/Source/Programs/WebTestsServer/dockerbuildandrun.bat
@@ -0,0 +1,14 @@
+@echo off
+
+:: Build image
+docker build . -t webtests
+
+set CONTAINER="webtestscontainer"
+:: Check if container exist and delete
+docker ps -a | findstr /c:"%CONTAINER%" >nul 2>&1
+if %errorlevel% equ 0 (
+	docker rm -f %CONTAINER%
+)
+
+:: Run container with image
+docker run --name %CONTAINER% -d -p 8000:8000 webtests
diff --git a/Engine/Source/Programs/WebTestsServer/dockerfile b/Engine/Source/Programs/WebTestsServer/dockerfile
new file mode 100644
--- /dev/null
+++ b/Engine/Source/Programs/WebTestsServer/dockerfile
@@ -0,0 +1,11 @@
+FROM python:3.10.4
+ENV DockerHOME=[REDACTED_USER_HOME]/webtests
+RUN mkdir -p $DockerHOME
+WORKDIR $DockerHOME
+ENV PYTHONDONTWRITEBYTECODE 1
+ENV PYTHONUNBUFFERED 1
+RUN pip install --upgrade pip
+COPY . $DockerHOME
+RUN pip install -r requirements.txt
+EXPOSE 8000
+CMD python manage.py runserver 0.0.0.0:8000
diff --git a/Engine/Source/Programs/WebTestsServer/httptests/admin.py b/Engine/Source/Programs/WebTestsServer/httptests/admin.py
new file mode 100644
--- /dev/null
+++ b/Engine/Source/Programs/WebTestsServer/httptests/admin.py
@@ -0,0 +1,3 @@
+from django.contrib import admin
+
+# Register your models here.
diff --git a/Engine/Source/Programs/WebTestsServer/httptests/apps.py b/Engine/Source/Programs/WebTestsServer/httptests/apps.py
new file mode 100644
--- /dev/null
+++ b/Engine/Source/Programs/WebTestsServer/httptests/apps.py
@@ -0,0 +1,6 @@
+from django.apps import AppConfig
+
+
+class HttpTestsConfig(AppConfig):
+    default_auto_field = 'django.db.models.BigAutoField'
+    name = 'httptests'
diff --git a/Engine/Source/Programs/WebTestsServer/httptests/models.py b/Engine/Source/Programs/WebTestsServer/httptests/models.py
new file mode 100644
--- /dev/null
+++ b/Engine/Source/Programs/WebTestsServer/httptests/models.py
@@ -0,0 +1,3 @@
+from django.db import models
+
+# Create your models here.
diff --git a/Engine/Source/Programs/WebTestsServer/httptests/tests.py b/Engine/Source/Programs/WebTestsServer/httptests/tests.py
new file mode 100644
--- /dev/null
+++ b/Engine/Source/Programs/WebTestsServer/httptests/tests.py
@@ -0,0 +1,3 @@
+from django.test import TestCase
+
+# Create your tests here.
diff --git a/Engine/Source/Programs/WebTestsServer/httptests/urls.py b/Engine/Source/Programs/WebTestsServer/httptests/urls.py
new file mode 100644
--- /dev/null
+++ b/Engine/Source/Programs/WebTestsServer/httptests/urls.py
@@ -0,0 +1,14 @@
+from django.urls import path
+
+from . import views
+
+urlpatterns = [
+    path("httptests/methods", views.methods, name="methods"),
+    path("httptests/get_large_response_without_chunks/<int:bytes_number>/", views.get_large_response_without_chunks, name="get_large_response_without_chunks"),
+    path("httptests/nonstreaming_receivetimeout/<int:wait_time>/", views.nonstreaming_receivetimeout, name="nonstreaming_receivetimeout"),
+    path("httptests/streaming_download/<int:chunks>/<int:chunk_size>/", views.streaming_download, name="streaming_download"),
+    path("httptests/streaming_upload_post", views.streaming_upload_post, name="streaming_upload_post"),
+    path("httptests/streaming_upload_put", views.streaming_upload_put, name="streaming_upload_put"),
+    path("httptests/redirect_from", views.redirect_from, name="redirect_from"),
+    path("httptests/redirect_to", views.redirect_to, name="redirect_to"),
+]
diff --git a/Engine/Source/Programs/WebTestsServer/httptests/views.py b/Engine/Source/Programs/WebTestsServer/httptests/views.py
new file mode 100644
--- /dev/null
+++ b/Engine/Source/Programs/WebTestsServer/httptests/views.py
@@ -0,0 +1,68 @@
+from django.shortcuts import render, redirect
+from django.http.response import JsonResponse
+from django.http import JsonResponse
+from django.http import StreamingHttpResponse
+from rest_framework.decorators import api_view, parser_classes
+from rest_framework.parsers import FileUploadParser
+import time
+import logging
+
+logger = logging.getLogger('django')
+
+@api_view(['GET', 'POST', 'DELETE', 'PUT'])
+def methods(request):
+    return JsonResponse({})
+
+@api_view(['GET'])
+def get_large_response_without_chunks(request, bytes_number):
+    data = "d" * bytes_number
+    return JsonResponse({'data' : data})
+
+@api_view(['GET'])
+def nonstreaming_receivetimeout(request, wait_time):
+    time.sleep(wait_time)
+    return JsonResponse({})
+
+@api_view(['GET'])
+def streaming_download(request, chunks, chunk_size):
+    def data_chunk_generator():
+        for x in range(chunks):
+            data_chunk = 'd' * chunk_size
+            yield data_chunk
+    response = StreamingHttpResponse(data_chunk_generator())
+    response['Content-Length'] = chunks * chunk_size
+    return response
+
+def output_file_info(f):
+    if f.multiple_chunks():
+        # big file
+        total_chunks = 0
+        total_bytes = 0
+        for chunk in f.chunks():
+            total_chunks += 1
+            total_bytes += len(chunk)
+        logger.info("Received %d chunks: %d bytes in total", total_chunks, total_bytes)
+    else:
+        # small file
+        logger.info("Received the whole file: %d bytes", f.size)
+
+@api_view(['POST'])
+def streaming_upload_post(request):
+    f = request.FILES['file']
+    output_file_info(f)
+    return JsonResponse({})
+
+@api_view(['PUT'])
+@parser_classes((FileUploadParser,))
+def streaming_upload_put(request):
+    f = request.data['file']
+    output_file_info(f)
+    return JsonResponse({})
+
+@api_view(['GET'])
+def redirect_from(request):
+    return redirect('redirect_to')
+
+@api_view(['GET'])
+def redirect_to(request):
+    return JsonResponse({})
diff --git a/Engine/Source/Programs/WebTestsServer/manage.py b/Engine/Source/Programs/WebTestsServer/manage.py
new file mode 100644
--- /dev/null
+++ b/Engine/Source/Programs/WebTestsServer/manage.py
@@ -0,0 +1,22 @@
+#!/usr/bin/env python
+"""Django's command-line utility for administrative tasks."""
+import os
+import sys
+
+
+def main():
+    """Run administrative tasks."""
+    os.environ.setdefault('DJANGO_SETTINGS_MODULE', 'webtests.settings')
+    try:
+        from django.core.management import execute_from_command_line
+    except ImportError as exc:
+        raise ImportError(
+            "Couldn't import Django. Are you sure it's installed and "
+            "available on your PYTHONPATH environment variable? Did you "
+            "forget to activate a virtual environment?"
+        ) from exc
+    execute_from_command_line(sys.argv)
+
+
+if __name__ == '__main__':
+    main()
diff --git a/Engine/Source/Programs/WebTestsServer/requirements.txt b/Engine/Source/Programs/WebTestsServer/requirements.txt
new file mode 100644
--- /dev/null
+++ b/Engine/Source/Programs/WebTestsServer/requirements.txt
@@ -0,0 +1,4 @@
+django==4.1.7
+djangorestframework==3.14.0
+channels==4.0.0
+daphne==4.0.0
diff --git a/Engine/Source/Programs/WebTestsServer/runserver.bat b/Engine/Source/Programs/WebTestsServer/runserver.bat
new file mode 100755
--- /dev/null
+++ b/Engine/Source/Programs/WebTestsServer/runserver.bat
@@ -0,0 +1,15 @@
+@echo off
+
+if exist env (
+	call .\env\Scripts\activate.bat
+) else (
+	..\..\..\Binaries\ThirdParty\Python3\Win64\python.exe -m venv env
+	call .\env\Scripts\activate.bat
+	pip install -r requirements.txt
+	if %ERRORLEVEL% NEQ 0 (
+		rmdir /s /q env
+		exit %ERRORLEVEL%
+	)
+)
+
+python manage.py runserver 0.0.0.0:8000
diff --git a/Engine/Source/Programs/WebTestsServer/runserver.sh b/Engine/Source/Programs/WebTestsServer/runserver.sh
new file mode 100644
--- /dev/null
+++ b/Engine/Source/Programs/WebTestsServer/runserver.sh
@@ -0,0 +1,13 @@
+if [ -d "env" ]; then
+	source ./env/Scripts/activate.sh
+else
+	run ../../../Binaries/ThirdParty/Python3/Mac/python3 -m venv env
+	source ./env/Scripts/activate.sh
+	run pip install -r requirements.txt
+	if [ $? -ne 0 ]; then
+		rm -rf env
+		exit $?
+	fi
+fi
+
+cmd python manage.py runserver 0.0.0.0:8000
diff --git a/Engine/Source/Programs/WebTestsServer/websocketstests/consumers.py b/Engine/Source/Programs/WebTestsServer/websocketstests/consumers.py
new file mode 100644
--- /dev/null
+++ b/Engine/Source/Programs/WebTestsServer/websocketstests/consumers.py
@@ -0,0 +1,13 @@
+import json
+
+from channels.generic.websocket import WebsocketConsumer
+
+class EchoConsumer(WebsocketConsumer):
+    def connect(self):
+        self.accept()
+
+    def disconnect(self, close_code):
+        pass
+
+    def receive(self, text_data):
+        self.send(text_data)
diff --git a/Engine/Source/Programs/WebTestsServer/websocketstests/routing.py b/Engine/Source/Programs/WebTestsServer/websocketstests/routing.py
new file mode 100644
--- /dev/null
+++ b/Engine/Source/Programs/WebTestsServer/websocketstests/routing.py
@@ -0,0 +1,7 @@
+from django.urls import re_path
+
+from . import consumers
+
+websocket_urlpatterns = [
+    re_path(r"webtests/websocketstests/echo", consumers.EchoConsumer.as_asgi()),
+]
diff --git a/Engine/Source/Programs/WebTestsServer/websocketstests/templates/websocketstests/index.html b/Engine/Source/Programs/WebTestsServer/websocketstests/templates/websocketstests/index.html
new file mode 100644
--- /dev/null
+++ b/Engine/Source/Programs/WebTestsServer/websocketstests/templates/websocketstests/index.html
@@ -0,0 +1,44 @@
+<!DOCTYPE html>
+<html>
+<head>
+    <meta charset="utf-8"/>
+    <title>WebSockects Tests</title>
+</head>
+<body>
+    <textarea id="echo-log" cols="100" rows="20"></textarea><br>
+    <input id="echo-message-input" type="text" size="100"><br>
+    <input id="echo-message-submit" type="button" value="Send">
+    <script>
+        const echoSocket = new WebSocket(
+            'ws://'
+            + window.location.host
+            + '/webtests/websocketstests/echo/'
+        );
+
+        echoSocket.onmessage = function(e) {
+            const data = JSON.parse(e.data);
+            document.querySelector('#echo-log').value += (data.message + '\n');
+        };
+
+        echoSocket.onclose = function(e) {
+            console.error('Echo socket closed unexpectedly');
+        };
+
+        document.querySelector('#echo-message-input').focus();
+        document.querySelector('#echo-message-input').onkeyup = function(e) {
+            if (e.key === 'Enter') {  // enter, return
+                document.querySelector('#echo-message-submit').click();
+            }
+        };
+
+        document.querySelector('#echo-message-submit').onclick = function(e) {
+            const messageInputDom = document.querySelector('#echo-message-input');
+            const message = messageInputDom.value;
+            echoSocket.send(JSON.stringify({
+                'message': message
+            }));
+            messageInputDom.value = '';
+        };
+    </script>
+</body>
+</html>
diff --git a/Engine/Source/Programs/WebTestsServer/websocketstests/urls.py b/Engine/Source/Programs/WebTestsServer/websocketstests/urls.py
new file mode 100644
--- /dev/null
+++ b/Engine/Source/Programs/WebTestsServer/websocketstests/urls.py
@@ -0,0 +1,7 @@
+from django.urls import path
+
+from . import views
+
+urlpatterns = [
+    path("websocketstests", views.index, name="index"),
+]
diff --git a/Engine/Source/Programs/WebTestsServer/websocketstests/views.py b/Engine/Source/Programs/WebTestsServer/websocketstests/views.py
new file mode 100644
--- /dev/null
+++ b/Engine/Source/Programs/WebTestsServer/websocketstests/views.py
@@ -0,0 +1,5 @@
+from django.shortcuts import render
+
+def index(request):
+    return render(request, "websocketstests/index.html")
+
diff --git a/Engine/Source/Programs/WebTestsServer/webtests/asgi.py b/Engine/Source/Programs/WebTestsServer/webtests/asgi.py
new file mode 100644
--- /dev/null
+++ b/Engine/Source/Programs/WebTestsServer/webtests/asgi.py
@@ -0,0 +1,28 @@
+"""
+ASGI config for webtests project.
+
+It exposes the ASGI callable as a module-level variable named ``application``.
+
+For more information on this file, see
+https://docs.djangoproject.com/en/4.1/howto/deployment/asgi/
+"""
+
+import os
+
+from channels.auth import AuthMiddlewareStack
+from channels.routing import ProtocolTypeRouter, URLRouter
+from channels.security.websocket import AllowedHostsOriginValidator
+from django.core.asgi import get_asgi_application
+
+from websocketstests.routing import websocket_urlpatterns
+
+os.environ.setdefault('DJANGO_SETTINGS_MODULE', 'webtests.settings')
+
+application = ProtocolTypeRouter(
+    {
+        "http": get_asgi_application(),
+        "websocket": AllowedHostsOriginValidator(
+            AuthMiddlewareStack(URLRouter(websocket_urlpatterns))
+        ),
+    }
+)
diff --git a/Engine/Source/Programs/WebTestsServer/webtests/settings.py b/Engine/Source/Programs/WebTestsServer/webtests/settings.py
new file mode 100644
--- /dev/null
+++ b/Engine/Source/Programs/WebTestsServer/webtests/settings.py
@@ -0,0 +1,134 @@
+"""
+Django settings for webtests project.
+
+Generated by 'django-admin startproject' using Django 4.1.7.
+
+For more information on this file, see
+https://docs.djangoproject.com/en/4.1/topics/settings/
+
+For the full list of settings and their values, see
+https://docs.djangoproject.com/en/4.1/ref/settings/
+"""
+
+from pathlib import Path
+
+# Build paths inside the project like this: BASE_DIR / 'subdir'.
+BASE_DIR = Path(__file__).resolve().parent.parent
+
+
+# Quick-start development settings - unsuitable for production
+# See https://docs.djangoproject.com/en/4.1/howto/deployment/checklist/
+
+# SECURITY WARNING: keep the secret key used in production secret!
+SECRET_KEY = 'django-insecure-x=3k^30px$8#ouewp0(1er-gpm@90k^z5$6bxx(qtot92yy$%w'
+
+# SECURITY WARNING: don't run with debug turned on in production!
+DEBUG = True
+
+# Allow any host, because this web server is only for local test purpose
+ALLOWED_HOSTS = ['*']
+
+
+# Application definition
+
+INSTALLED_APPS = [
+    'daphne',
+    'websocketstests',
+    'django.contrib.admin',
+    'django.contrib.auth',
+    'django.contrib.contenttypes',
+    'django.contrib.sessions',
+    'django.contrib.messages',
+    'django.contrib.staticfiles',
+]
+
+MIDDLEWARE = [
+    'django.middleware.security.SecurityMiddleware',
+    'django.contrib.sessions.middleware.SessionMiddleware',
+    'django.middleware.common.CommonMiddleware',
+    'django.middleware.csrf.CsrfViewMiddleware',
+    'django.contrib.auth.middleware.AuthenticationMiddleware',
+    'django.contrib.messages.middleware.MessageMiddleware',
+    'django.middleware.clickjacking.XFrameOptionsMiddleware',
+]
+
+ROOT_URLCONF = 'webtests.urls'
+
+TEMPLATES = [
+    {
+        'BACKEND': 'django.template.backends.django.DjangoTemplates',
+        'DIRS': [],
+        'APP_DIRS': True,
+        'OPTIONS': {
+            'context_processors': [
+                'django.template.context_processors.debug',
+                'django.template.context_processors.request',
+                'django.contrib.auth.context_processors.auth',
+                'django.contrib.messages.context_processors.messages',
+            ],
+        },
+    },
+]
+
+WSGI_APPLICATION = 'webtests.wsgi.application'
+
+
+# Database
+# https://docs.djangoproject.com/en/4.1/ref/settings/#databases
+
+DATABASES = {
+    'default': {
+        'ENGINE': 'django.db.backends.sqlite3',
+        'NAME': BASE_DIR / 'db.sqlite3',
+    }
+}
+
+
+# Password validation
+# https://docs.djangoproject.com/en/4.1/ref/settings/#auth-password-validators
+
+AUTH_PASSWORD_VALIDATORS = [
+    {
+        'NAME': 'django.contrib.auth.password_validation.UserAttributeSimilarityValidator',
+    },
+    {
+        'NAME': 'django.contrib.auth.password_validation.MinimumLengthValidator',
+    },
+    {
+        'NAME': 'django.contrib.auth.password_validation.CommonPasswordValidator',
+    },
+    {
+        'NAME': 'django.contrib.auth.password_validation.NumericPasswordValidator',
+    },
+]
+
+
+# Internationalization
+# https://docs.djangoproject.com/en/4.1/topics/i18n/
+
+LANGUAGE_CODE = 'en-us'
+
+TIME_ZONE = 'UTC'
+
+USE_I18N = True
+
+USE_TZ = True
+
+
+# Static files (CSS, JavaScript, Images)
+# https://docs.djangoproject.com/en/4.1/howto/static-files/
+
+STATIC_URL = 'static/'
+
+# Default primary key field type
+# https://docs.djangoproject.com/en/4.1/ref/settings/#default-auto-field
+
+DEFAULT_AUTO_FIELD = 'django.db.models.BigAutoField'
+
+LOGGING = {
+    "version": 1,  # the dictConfig format version
+    "disable_existing_loggers": False,  # retain the default loggers
+    "level": "DEBUG",
+}
+
+ASGI_APPLICATION = "webtests.asgi.application"
diff --git a/Engine/Source/Programs/WebTestsServer/webtests/urls.py b/Engine/Source/Programs/WebTestsServer/webtests/urls.py
new file mode 100644
--- /dev/null
+++ b/Engine/Source/Programs/WebTestsServer/webtests/urls.py
@@ -0,0 +1,23 @@
+"""webtests URL Configuration
+
+The `urlpatterns` list routes URLs to views. For more information please see:
+    https://docs.djangoproject.com/en/4.1/topics/http/urls/
+Examples:
+Function views
+    1. Add an import:  from my_app import views
+    2. Add a URL to urlpatterns:  path('', views.home, name='home')
+Class-based views
+    1. Add an import:  from other_app.views import Home
+    2. Add a URL to urlpatterns:  path('', Home.as_view(), name='home')
+Including another URLconf
+    1. Import the include() function: from django.urls import include, path
+    2. Add a URL to urlpatterns:  path('blog/', include('blog.urls'))
+"""
+from django.contrib import admin
+from django.urls import include, path
+
+urlpatterns = [
+    path("webtests/", include("httptests.urls")),
+    path("webtests/", include("websocketstests.urls")),
+    path('admin/', admin.site.urls),
+]
diff --git a/Engine/Source/Programs/WebTestsServer/webtests/wsgi.py b/Engine/Source/Programs/WebTestsServer/webtests/wsgi.py
new file mode 100644
--- /dev/null
+++ b/Engine/Source/Programs/WebTestsServer/webtests/wsgi.py
@@ -0,0 +1,16 @@
+"""
+WSGI config for webtests project.
+
+It exposes the WSGI callable as a module-level variable named ``application``.
+
+For more information on this file, see
+https://docs.djangoproject.com/en/4.1/howto/deployment/wsgi/
+"""
+
+import os
+
+from django.core.wsgi import get_wsgi_application
+
+os.environ.setdefault('DJANGO_SETTINGS_MODULE', 'webtests.settings')
+
+application = get_wsgi_application()
diff --git a/Engine/Source/Runtime/Online/WebSockets/Private/WebSocketsModule.cpp b/Engine/Source/Runtime/Online/WebSockets/Private/WebSocketsModule.cpp
--- a/Engine/Source/Runtime/Online/WebSockets/Private/WebSocketsModule.cpp
+++ b/Engine/Source/Runtime/Online/WebSockets/Private/WebSocketsModule.cpp
@@ -33,6 +33,11 @@ void FWebSocketsModule::StartupModule()
 	// Default configuration values can be found in BaseEngine.ini
 	TArray<FString> Protocols;
 	GConfig->GetArray(TEXT("WebSockets"), TEXT("WebSocketsProtocols"), Protocols, GEngineIni);
+	if (Protocols.IsEmpty())
+	{
+		Protocols.Add(TEXT("ws"));
+		Protocols.Add(TEXT("wss"));
+	}
 
 	WebSocketsManager = new FPlatformWebSocketsManager;
 	WebSocketsManager->InitWebSockets(Protocols);
