Upgrade pip package in python env to fix the pip outdated warning.

diff --git a/Engine/Source/Programs/WebTestsServer/runserver.bat b/Engine/Source/Programs/WebTestsServer/runserver.bat
--- a/Engine/Source/Programs/WebTestsServer/runserver.bat
+++ b/Engine/Source/Programs/WebTestsServer/runserver.bat
@@ -5,6 +5,7 @@ if exist env (
 ) else (
 	..\..\..\Binaries\ThirdParty\Python3\Win64\python.exe -m venv env
 	call .\env\Scripts\activate.bat
+	python -m pip install --upgrade pip
 	pip install -r requirements.txt
 	if %ERRORLEVEL% NEQ 0 (
 		rmdir /s /q env
diff --git a/Engine/Source/Programs/WebTestsServer/runserver.sh b/Engine/Source/Programs/WebTestsServer/runserver.sh
--- a/Engine/Source/Programs/WebTestsServer/runserver.sh
+++ b/Engine/Source/Programs/WebTestsServer/runserver.sh
@@ -3,6 +3,7 @@ if [ -d "env" ]; then
 else
 	run ../../../Binaries/ThirdParty/Python3/Mac/python3 -m venv env
 	source ./env/Scripts/activate.sh
+	run python -m  pip install --upgrade pip
 	run pip install -r requirements.txt
 	if [ $? -ne 0 ]; then
 		rm -rf env
