Fix the package install error in daphne4.0.0 when run WebTestsServer on linux, after updating python to 3.11 in 32683702: ImportError: libffi.so.6: cannot open shared object file: No such file or directory

diff --git a/Engine/Source/Programs/WebTestsServer/createenv.sh b/Engine/Source/Programs/WebTestsServer/createenv.sh
--- a/Engine/Source/Programs/WebTestsServer/createenv.sh
+++ b/Engine/Source/Programs/WebTestsServer/createenv.sh
@@ -4,6 +4,9 @@ if [ ! -d "env" ]; then
 			../../../Binaries/ThirdParty/Python3/Mac/bin/python3 -m venv env
 			;;
 		*)
+			# For daphne4.0.0 -> ImportError: libffi.so.6: cannot open shared object file: No such file or directory
+			curl -LO http://archive.ubuntu.com/ubuntu/pool/main/libf/libffi/libffi6_3.2.1-8_amd64.deb
+			sudo dpkg -i libffi6_3.2.1-8_amd64.deb
 			../../../Binaries/ThirdParty/Python3/Linux/bin/python3 -m venv env
 			;;
 	esac
