Freeze working version of autobahn in requirements.txt of WebTestsServer to fix Mac websockets error when start: No module named '_nvx_utf8validator'.

diff --git a/Engine/Source/Programs/WebTestsServer/requirements.txt b/Engine/Source/Programs/WebTestsServer/requirements.txt
--- a/Engine/Source/Programs/WebTestsServer/requirements.txt
+++ b/Engine/Source/Programs/WebTestsServer/requirements.txt
@@ -1,3 +1,4 @@
+autobahn==24.4.2
 django==4.2.0
 djangorestframework==3.14.0
 channels==4.0.0
